Add UnrealToolboxExampleApplication: a multi-page developer guide for publishing Toolbox tools

A new Avalonia UI application that serves as both a working example Toolbox tool and
interactive documentation for tool authors. Uses FluentAvalonia NavigationView with a
six-page progressive guide covering project setup, Toolbox.json configuration,
install/uninstall lifecycle, publishing to Horde, and a self-referential source walkthrough.

Application features:
- NavigationView sidebar with lazy page creation and caching
- Shared styles (DocCard, CodeBlock, PageTitle, SectionHeader) using FluentAvaloniaUI resources
- Selectable code blocks with monospace font for copyable snippets
- Cross-platform single-instance lock using file-based locking
- Graceful update handshake via named events (Windows) and marker files (macOS/Linux)
- Platform-specific install: Start Menu shortcut, macOS symlink, Linux .desktop entry
- Platform-specific Toolbox.json variants (win/osx/linux) selected via .csproj RID conditions
- Tray popup menu item ("Open Developer Guide") demonstrating popupMenu in Toolbox.json
- Application icon (.ico/.icns) and banner image on the Welcome page
- Build definition for automated deployment (UnrealToolboxExampleApplicationBuild.xml)

#Synced-CL 51609771

#Synced-CL 51609771

#Synced-CL 51609771

diff --git a/Engine/Source/Programs/UnrealToolboxExampleApplication/App.axaml b/Engine/Source/Programs/UnrealToolboxExampleApplication/App.axaml
new file mode 100644
--- /dev/null
+++ b/Engine/Source/Programs/UnrealToolboxExampleApplication/App.axaml
@@ -0,0 +1,10 @@
+<Application xmlns="[REDACTED_PROFILE_URL]"
+             xmlns:x="http://schemas.microsoft.com/winfx/2006/xaml"
+             xmlns:sty="using:FluentAvalonia.Styling"
+             x:Class="UnrealToolboxExampleApplication.App"
+             RequestedThemeVariant="Default">
+    <Application.Styles>
+        <sty:FluentAvaloniaTheme PreferSystemTheme="True" PreferUserAccentColor="True" />
+        <StyleInclude Source="/Styles.axaml"/>
+    </Application.Styles>
+</Application>
diff --git a/Engine/Source/Programs/UnrealToolboxExampleApplication/App.axaml.cs b/Engine/Source/Programs/UnrealToolboxExampleApplication/App.axaml.cs
new file mode 100644
--- /dev/null
+++ b/Engine/Source/Programs/UnrealToolboxExampleApplication/App.axaml.cs
@@ -0,0 +1,148 @@
+// Copyright Epic Games, Inc. All Rights Reserved.
+
+using Avalonia;
+using Avalonia.Controls;
+using Avalonia.Controls.ApplicationLifetimes;
+using Avalonia.Markup.Xaml;
+using Avalonia.Threading;
+
+namespace UnrealToolboxExampleApplication;
+
+public partial class App : Application
+{
+	public override void Initialize()
+	{
+		AvaloniaXamlLoader.Load(this);
+	}
+
+	public override void OnFrameworkInitializationCompleted()
+	{
+		if (ApplicationLifetime is IClassicDesktopStyleApplicationLifetime desktop)
+		{
+			desktop.MainWindow = new MainWindow();
+			StartUpdateListener(desktop);
+		}
+
+		base.OnFrameworkInitializationCompleted();
+	}
+
+	void StartUpdateListener(IClassicDesktopStyleApplicationLifetime desktop)
+	{
+		if (OperatingSystem.IsWindows())
+		{
+			Thread thread = new Thread(() => ListenViaNamedEvent(desktop));
+			thread.IsBackground = true;
+			thread.Start();
+		}
+		else
+		{
+			Thread thread = new Thread(() => ListenViaMarkerFile(desktop));
+			thread.IsBackground = true;
+			thread.Start();
+		}
+	}
+
+	void ListenViaNamedEvent(IClassicDesktopStyleApplicationLifetime desktop)
+	{
+		using EventWaitHandle requestEvent = new EventWaitHandle(false, EventResetMode.AutoReset, Program.UpdateRequestEventName);
+
+		while (true)
+		{
+			requestEvent.WaitOne();
+
+			ManualResetEventSlim dialogComplete = new ManualResetEventSlim(false);
+
+			Dispatcher.UIThread.Post(async () =>
+			{
+				try
+				{
+					bool accepted = await ShowUpdateDialog(desktop);
+					if (accepted)
+					{
+						desktop.Shutdown();
+					}
+					else
+					{
+						using EventWaitHandle declinedEvent = new EventWaitHandle(false, EventResetMode.AutoReset, Program.UpdateDeclinedEventName);
+						declinedEvent.Set();
+					}
+				}
+				finally
+				{
+					dialogComplete.Set();
+				}
+			});
+
+			dialogComplete.Wait();
+			dialogComplete.Dispose();
+		}
+	}
+
+	void ListenViaMarkerFile(IClassicDesktopStyleApplicationLifetime desktop)
+	{
+		string requestFilePath = Program.GetRequestFilePath();
+
+		while (true)
+		{
+			Thread.Sleep(1000);
+
+			if (!File.Exists(requestFilePath))
+			{
+				continue;
+			}
+
+			try
+			{
+				File.Delete(requestFilePath);
+			}
+			catch
+			{
+				continue;
+			}
+
+			ManualResetEventSlim dialogComplete = new ManualResetEventSlim(false);
+
+			Dispatcher.UIThread.Post(async () =>
+			{
+				try
+				{
+					bool accepted = await ShowUpdateDialog(desktop);
+					if (accepted)
+					{
+						desktop.Shutdown();
+					}
+					else
+					{
+						try
+						{
+							File.WriteAllText(Program.GetDeclinedFilePath(), "declined");
+						}
+						catch
+						{
+						}
+					}
+				}
+				finally
+				{
+					dialogComplete.Set();
+				}
+			});
+
+			dialogComplete.Wait();
+			dialogComplete.Dispose();
+		}
+	}
+
+	static async Task<bool> ShowUpdateDialog(IClassicDesktopStyleApplicationLifetime desktop)
+	{
+		Window? mainWindow = desktop.MainWindow;
+		if (mainWindow == null)
+		{
+			return false;
+		}
+
+		UpdateDialog dialog = new UpdateDialog();
+		bool? result = await dialog.ShowDialog<bool?>(mainWindow);
+		return result == true;
+	}
+}
diff --git a/Engine/Source/Programs/UnrealToolboxExampleApplication/CodeSamples.cs b/Engine/Source/Programs/UnrealToolboxExampleApplication/CodeSamples.cs
new file mode 100644
--- /dev/null
+++ b/Engine/Source/Programs/UnrealToolboxExampleApplication/CodeSamples.cs
@@ -0,0 +1,296 @@
+// Copyright Epic Games, Inc. All Rights Reserved.
+
+namespace UnrealToolboxExampleApplication;
+
+static class CodeSamples
+{
+	public static readonly string CsprojStructure = """
+		<Project Sdk="Microsoft.NET.Sdk">
+		  <PropertyGroup>
+		    <OutputType>WinExe</OutputType>
+		    <TargetFramework>net10.0</TargetFramework>
+		    <Nullable>enable</Nullable>
+		    <AvaloniaUseCompiledBindingsByDefault>true</AvaloniaUseCompiledBindingsByDefault>
+		  </PropertyGroup>
+
+		  <ItemGroup>
+		    <PackageReference Include="Avalonia" Version="11.3.12" />
+		    <PackageReference Include="Avalonia.Desktop" Version="11.3.12" />
+		    <PackageReference Include="Avalonia.Themes.Fluent" Version="11.3.12" />
+		    <PackageReference Include="FluentAvaloniaUI" Version="2.5.0" />
+		  </ItemGroup>
+		</Project>
+		""";
+
+	public static readonly string SingleInstanceLock = """
+		class SingleInstanceLock : IDisposable
+		{
+		    readonly string _lockFilePath;
+		    FileStream? _lockFileStream;
+
+		    public SingleInstanceLock(string appName)
+		    {
+		        string lockDir = Path.Combine(
+		            Environment.GetFolderPath(
+		                Environment.SpecialFolder.LocalApplicationData),
+		            "Epic Games", appName);
+		        Directory.CreateDirectory(lockDir);
+		        _lockFilePath = Path.Combine(lockDir, ".lock");
+		    }
+
+		    public bool TryAcquire()
+		    {
+		        try
+		        {
+		            _lockFileStream = new FileStream(
+		                _lockFilePath, FileMode.OpenOrCreate,
+		                FileAccess.ReadWrite, FileShare.None);
+		            return true;
+		        }
+		        catch (IOException) { return false; }
+		    }
+
+		    public void Dispose()
+		    {
+		        _lockFileStream?.Dispose();
+		        _lockFileStream = null;
+		    }
+		}
+		""";
+
+	public static readonly string SingleInstanceUsage = """
+		using SingleInstanceLock instanceLock =
+		    new SingleInstanceLock("MyApp");
+		if (!instanceLock.TryAcquire())
+		{
+		    Console.Error.WriteLine("Already running.");
+		    return 1;
+		}
+		""";
+
+	public static readonly string PlatformRidConditions = """
+		<!-- Default to Windows for local build -->
+		<ItemGroup Condition="'$(RuntimeIdentifier)' == ''
+		    Or $(RuntimeIdentifier.StartsWith('win'))">
+		  <Content Include="Toolbox.win.json" Link="Toolbox.json">
+		    <CopyToOutputDirectory>PreserveNewest</CopyToOutputDirectory>
+		    <CopyToPublishDirectory>PreserveNewest</CopyToPublishDirectory>
+		  </Content>
+		</ItemGroup>
+
+		<ItemGroup Condition="$(RuntimeIdentifier.StartsWith('osx'))">
+		  <Content Include="Toolbox.osx.json" Link="Toolbox.json">
+		    <CopyToOutputDirectory>PreserveNewest</CopyToOutputDirectory>
+		    <CopyToPublishDirectory>PreserveNewest</CopyToPublishDirectory>
+		  </Content>
+		</ItemGroup>
+
+		<ItemGroup Condition="$(RuntimeIdentifier.StartsWith('linux'))">
+		  <Content Include="Toolbox.linux.json" Link="Toolbox.json">
+		    <CopyToOutputDirectory>PreserveNewest</CopyToOutputDirectory>
+		    <CopyToPublishDirectory>PreserveNewest</CopyToPublishDirectory>
+		  </Content>
+		</ItemGroup>
+		""";
+
+	public static readonly string ToolboxJsonFull = """
+		{
+		  "installCommand": {
+		    "fileName": "MyTool.exe",
+		    "arguments": [ "--install" ]
+		  },
+		  "uninstallCommand": {
+		    "fileName": "MyTool.exe",
+		    "arguments": [ "--uninstall" ]
+		  },
+		  "runCommand": {
+		    "fileName": "MyTool.exe"
+		  },
+		  "manualInstall": false,
+		  "requiresElevation": false,
+		  "hidden": false,
+		  "popupMenu": {
+		    "label": "My Tool",
+		    "children": [
+		      {
+		        "label": "Open Dashboard",
+		        "command": {
+		          "fileName": "MyTool.exe",
+		          "arguments": [ "--dashboard" ]
+		        }
+		      },
+		      {
+		        "label": "View Logs",
+		        "command": {
+		          "fileName": "MyTool.exe",
+		          "arguments": [ "--logs" ]
+		        }
+		      }
+		    ]
+		  }
+		}
+		""";
+
+	public static readonly string ToolboxJsonMinimal = """
+		{
+		  "installCommand": {
+		    "fileName": "MyTool.exe",
+		    "arguments": [ "--install" ]
+		  },
+		  "uninstallCommand": {
+		    "fileName": "MyTool.exe",
+		    "arguments": [ "--uninstall" ]
+		  }
+		}
+		""";
+
+	public static readonly string InstallWindows = """
+		static void InstallWindows(string exePath)
+		{
+		    string shortcutPath = GetWindowsShortcutPath();
+		    Type shellType = Type.GetTypeFromProgID("WScript.Shell")
+		        ?? throw new InvalidOperationException(
+		            "WScript.Shell is not available.");
+		    dynamic shell = Activator.CreateInstance(shellType)!;
+		    dynamic shortcut = shell.CreateShortcut(shortcutPath);
+		    shortcut.TargetPath = exePath;
+		    shortcut.WorkingDirectory =
+		        Path.GetDirectoryName(exePath);
+		    shortcut.Description = AppDisplayName;
+		    shortcut.Save();
+
+		    Marshal.ReleaseComObject(shortcut);
+		    Marshal.ReleaseComObject(shell);
+		}
+		""";
+
+	public static readonly string InstallMacOs = """
+		static void InstallMacOS(string exePath)
+		{
+		    string applicationsDir = Path.Combine(
+		        Environment.GetFolderPath(
+		            Environment.SpecialFolder.UserProfile),
+		        "Applications");
+		    Directory.CreateDirectory(applicationsDir);
+
+		    string linkPath = Path.Combine(
+		        applicationsDir, ProcessName);
+		    if (File.Exists(linkPath))
+		        File.Delete(linkPath);
+
+		    File.CreateSymbolicLink(linkPath, exePath);
+		}
+		""";
+
+	public static readonly string InstallLinux = """
+		static void InstallLinux(string exePath)
+		{
+		    string applicationsDir = Path.Combine(
+		        Environment.GetFolderPath(
+		            Environment.SpecialFolder.UserProfile),
+		        ".local", "share", "applications");
+		    Directory.CreateDirectory(applicationsDir);
+
+		    string desktopFilePath = Path.Combine(
+		        applicationsDir, $"{ProcessName}.desktop");
+		    string content =
+		        "[Desktop Entry]\n" +
+		        $"Name={AppDisplayName}\n" +
+		        $"Exec={exePath}\n" +
+		        "Type=Application\n" +
+		        "Terminal=false\n";
+
+		    File.WriteAllText(desktopFilePath, content);
+		}
+		""";
+
+	public static readonly string NamedEventUpdate = """
+		// Installer sends update request
+		using EventWaitHandle requestEvent =
+		    new EventWaitHandle(false, EventResetMode.AutoReset,
+		        "MyApp-UpdateRequest");
+		requestEvent.Set();
+
+		// Running app listens and shows dialog
+		requestEvent.WaitOne();
+		bool accepted = await ShowUpdateDialog();
+		if (accepted)
+		    desktop.Shutdown();
+		else
+		{
+		    using EventWaitHandle declinedEvent =
+		        new EventWaitHandle(false, EventResetMode.AutoReset,
+		            "MyApp-UpdateDeclined");
+		    declinedEvent.Set();
+		}
+		""";
+
+	public static readonly string MarkerFileUpdate = """
+		// Installer writes marker file
+		File.WriteAllText(requestFilePath, "update");
+
+		// Running app polls for the file
+		while (true)
+		{
+		    Thread.Sleep(1000);
+		    if (!File.Exists(requestFilePath))
+		        continue;
+
+		    File.Delete(requestFilePath);
+		    bool accepted = await ShowUpdateDialog();
+		    if (!accepted)
+		        File.WriteAllText(declinedFilePath, "declined");
+		    else
+		        desktop.Shutdown();
+		}
+		""";
+
+	public static readonly string ToolboxWinJson = """
+		{
+		  "installCommand": {
+		    "fileName": "UnrealToolboxExampleApplication.exe",
+		    "arguments": [ "--install" ]
+		  },
+		  "uninstallCommand": {
+		    "fileName": "UnrealToolboxExampleApplication.exe",
+		    "arguments": [ "--uninstall" ]
+		  },
+		  "runCommand": {
+		    "fileName": "UnrealToolboxExampleApplication.exe"
+		  },
+		  "popupMenu": {
+		    "label": "Open Developer Guide",
+		    "command": {
+		      "fileName": "UnrealToolboxExampleApplication.exe"
+		    }
+		  }
+		}
+		""";
+
+	public static readonly string ToolboxUnixJson = """
+		{
+		  "installCommand": {
+		    "fileName": "UnrealToolboxExampleApplication",
+		    "arguments": [ "--install" ]
+		  },
+		  "uninstallCommand": {
+		    "fileName": "UnrealToolboxExampleApplication",
+		    "arguments": [ "--uninstall" ]
+		  },
+		  "runCommand": {
+		    "fileName": "UnrealToolboxExampleApplication"
+		  },
+		  "popupMenu": {
+		    "label": "Open Developer Guide",
+		    "command": {
+		      "fileName": "UnrealToolboxExampleApplication"
+		    }
+		  }
+		}
+		""";
+
+	public static readonly string LocalToolFlag = """
+		dotnet run --project UnrealToolbox.csproj -- \
+		    -NoUpdate -LocalTool=C:\dev\MyTool\bin\publish
+		""";
+}
diff --git a/Engine/Source/Programs/UnrealToolboxExampleApplication/MainWindow.axaml b/Engine/Source/Programs/UnrealToolboxExampleApplication/MainWindow.axaml
new file mode 100644
--- /dev/null
+++ b/Engine/Source/Programs/UnrealToolboxExampleApplication/MainWindow.axaml
@@ -0,0 +1,21 @@
+<Window xmlns="[REDACTED_PROFILE_URL]"
+        xmlns:x="http://schemas.microsoft.com/winfx/2006/xaml"
+        xmlns:ui="using:FluentAvalonia.UI.Controls"
+        x:Class="UnrealToolboxExampleApplication.MainWindow"
+        Title="Unreal Toolbox Example Application"
+        Width="900"
+        Height="650"
+        WindowStartupLocation="CenterScreen"
+        Icon="/Resources/utbea.ico">
+
+    <Panel>
+        <ui:NavigationView x:Name="_navView"
+                           OpenPaneLength="260"
+                           IsSettingsVisible="False"
+                           PaneDisplayMode="Left"
+                           PaneTitle="Developer Guide"
+                           BorderBrush="Transparent">
+        </ui:NavigationView>
+    </Panel>
+
+</Window>
diff --git a/Engine/Source/Programs/UnrealToolboxExampleApplication/MainWindow.axaml.cs b/Engine/Source/Programs/UnrealToolboxExampleApplication/MainWindow.axaml.cs
new file mode 100644
--- /dev/null
+++ b/Engine/Source/Programs/UnrealToolboxExampleApplication/MainWindow.axaml.cs
@@ -0,0 +1,58 @@
+// Copyright Epic Games, Inc. All Rights Reserved.
+
+using Avalonia.Controls;
+using FluentAvalonia.UI.Controls;
+using UnrealToolboxExampleApplication.Pages;
+
+namespace UnrealToolboxExampleApplication;
+
+public partial class MainWindow : Window
+{
+	readonly Dictionary<string, UserControl> _pageCache = new Dictionary<string, UserControl>();
+
+	public MainWindow()
+	{
+		InitializeComponent();
+
+		_navView.MenuItems.Add(new NavigationViewItem() { Content = "Welcome", IconSource = new SymbolIconSource() { Symbol = Symbol.Home }, Tag = "Welcome" });
+		_navView.MenuItems.Add(new NavigationViewItem() { Content = "Project Setup", IconSource = new SymbolIconSource() { Symbol = Symbol.Code }, Tag = "ProjectSetup" });
+		_navView.MenuItems.Add(new NavigationViewItem() { Content = "Toolbox.json Reference", IconSource = new SymbolIconSource() { Symbol = Symbol.Document }, Tag = "ToolboxJson" });
+		_navView.MenuItems.Add(new NavigationViewItem() { Content = "Install & Uninstall", IconSource = new SymbolIconSource() { Symbol = Symbol.Sync }, Tag = "InstallUninstall" });
+		_navView.MenuItems.Add(new NavigationViewItem() { Content = "Publishing to Horde", IconSource = new SymbolIconSource() { Symbol = Symbol.Globe }, Tag = "Publishing" });
+		_navView.MenuItems.Add(new NavigationViewItem() { Content = "This Example", IconSource = new SymbolIconSource() { Symbol = Symbol.Important }, Tag = "ThisExample" });
+
+		_navView.SelectionChanged += NavView_SelectionChanged;
+		_navView.SelectedItem = _navView.MenuItems[0];
+	}
+
+	void NavView_SelectionChanged(object? sender, NavigationViewSelectionChangedEventArgs e)
+	{
+		if (_navView.SelectedItem is NavigationViewItem nvi)
+		{
+			string? tag = nvi.Tag as string;
+			if (tag != null)
+			{
+				_navView.Content = GetOrCreatePage(tag);
+			}
+		}
+	}
+
+	UserControl GetOrCreatePage(string tag)
+	{
+		if (!_pageCache.TryGetValue(tag, out UserControl? page))
+		{
+			page = tag switch
+			{
+				"Welcome" => new WelcomePage(),
+				"ProjectSetup" => new ProjectSetupPage(),
+				"ToolboxJson" => new ToolboxJsonPage(),
+				"InstallUninstall" => new InstallUninstallPage(),
+				"Publishing" => new PublishingPage(),
+				"ThisExample" => new ThisExamplePage(),
+				_ => new WelcomePage()
+			};
+			_pageCache.Add(tag, page);
+		}
+		return page;
+	}
+}
diff --git a/Engine/Source/Programs/UnrealToolboxExampleApplication/Pages/InstallUninstallPage.axaml b/Engine/Source/Programs/UnrealToolboxExampleApplication/Pages/InstallUninstallPage.axaml
new file mode 100644
--- /dev/null
+++ b/Engine/Source/Programs/UnrealToolboxExampleApplication/Pages/InstallUninstallPage.axaml
@@ -0,0 +1,112 @@
+<UserControl xmlns="[REDACTED_PROFILE_URL]"
+             xmlns:x="http://schemas.microsoft.com/winfx/2006/xaml"
+             xmlns:local="using:UnrealToolboxExampleApplication"
+             x:Class="UnrealToolboxExampleApplication.Pages.InstallUninstallPage">
+
+    <ScrollViewer Padding="32 24">
+        <StackPanel Spacing="4" MaxWidth="700">
+
+            <TextBlock Classes="PageTitle" Text="Install &amp; Uninstall"/>
+
+            <!-- Installation flow -->
+            <TextBlock Classes="SectionHeader" Text="Installation Flow"/>
+
+            <TextBlock Classes="DocBody"
+                       Text="When the Toolbox installs a tool, it follows these steps:"/>
+
+            <Border Classes="DocCard">
+                <StackPanel Spacing="4">
+                    <TextBlock Classes="DocBody" Text="1. Download the tool zip from the Horde server"/>
+                    <TextBlock Classes="DocBody" Text="2. Extract to a versioned directory under the Toolbox data folder"/>
+                    <TextBlock Classes="DocBody" Text="3. Parse Toolbox.json from the extracted root"/>
+                    <TextBlock Classes="DocBody" Text="4. Run installCommand with the tool directory as the working directory"/>
+                    <TextBlock Classes="DocBody" Text="5. Record the installed version in Tools.json"/>
+                </StackPanel>
+            </Border>
+
+            <!-- Windows install -->
+            <TextBlock Classes="SectionHeader" Text="Windows: Start Menu Shortcut"/>
+
+            <TextBlock Classes="DocBody"
+                       Text="This example creates a Start Menu shortcut using COM automation (WScript.Shell). The shortcut points to the tool executable and sets the working directory:"/>
+
+            <Border Classes="CodeBlock">
+                <SelectableTextBlock Text="{x:Static local:CodeSamples.InstallWindows}"/>
+            </Border>
+
+            <!-- macOS install -->
+            <TextBlock Classes="SectionHeader" Text="macOS: Symlink in ~/Applications"/>
+
+            <TextBlock Classes="DocBody"
+                       Text="On macOS, a symbolic link is created in the user's Applications directory, providing Finder and Spotlight discoverability:"/>
+
+            <Border Classes="CodeBlock">
+                <SelectableTextBlock Text="{x:Static local:CodeSamples.InstallMacOs}"/>
+            </Border>
+
+            <!-- Linux install -->
+            <TextBlock Classes="SectionHeader" Text="Linux: .desktop Entry"/>
+
+            <TextBlock Classes="DocBody"
+                       Text="On Linux, a .desktop file is written to ~/.local/share/applications, making the tool appear in application launchers:"/>
+
+            <Border Classes="CodeBlock">
+                <SelectableTextBlock Text="{x:Static local:CodeSamples.InstallLinux}"/>
+            </Border>
+
+            <!-- Graceful update -->
+            <TextBlock Classes="SectionHeader" Text="Graceful Update Handshake"/>
+
+            <TextBlock Classes="DocBody"
+                       Text="When the Toolbox updates a tool, it runs the uninstallCommand first. If the tool is currently running, the uninstall process needs to ask it to shut down. This example implements a two-party handshake:"/>
+
+            <Border Classes="DocCard">
+                <StackPanel Spacing="4">
+                    <TextBlock FontWeight="SemiBold" Text="Windows: Named Events"/>
+                    <TextBlock Classes="DocBody"
+                               Text="The uninstaller signals a named event. The running app listens on a background thread, shows an update dialog, and either shuts down (accepted) or signals a declined event."/>
+                </StackPanel>
+            </Border>
+
+            <Border Classes="CodeBlock">
+                <SelectableTextBlock Text="{x:Static local:CodeSamples.NamedEventUpdate}"/>
+            </Border>
+
+            <Border Classes="DocCard">
+                <StackPanel Spacing="4">
+                    <TextBlock FontWeight="SemiBold" Text="macOS/Linux: Marker Files"/>
+                    <TextBlock Classes="DocBody"
+                               Text="Named events are not available on Unix. Instead, the uninstaller writes a marker file to the temp directory. The running app polls for this file every second and responds the same way."/>
+                </StackPanel>
+            </Border>
+
+            <Border Classes="CodeBlock">
+                <SelectableTextBlock Text="{x:Static local:CodeSamples.MarkerFileUpdate}"/>
+            </Border>
+
+            <!-- MSI tools -->
+            <TextBlock Classes="SectionHeader" Text="MSI Tools"/>
+
+            <Border Classes="DocCard">
+                <StackPanel Spacing="4">
+                    <TextBlock FontWeight="SemiBold" Text="msi-product-id Metadata"/>
+                    <TextBlock Classes="DocBody"
+                               Text="If your tool is distributed as an MSI installer, set the msi-product-id metadata on the Horde tool deployment. The Toolbox will use this to check whether the MSI is already installed and skip redundant installs. MSI tools are not auto-updated."/>
+                </StackPanel>
+            </Border>
+
+            <!-- Elevation -->
+            <TextBlock Classes="SectionHeader" Text="Elevation"/>
+
+            <Border Classes="DocCard">
+                <StackPanel Spacing="4">
+                    <TextBlock FontWeight="SemiBold" Text="requiresElevation"/>
+                    <TextBlock Classes="DocBody"
+                               Text="When set to true in Toolbox.json, the installCommand and uninstallCommand are run with elevated privileges. On Windows this triggers the UAC prompt; on macOS/Linux this uses sudo. Use this only when your install step requires writing to protected locations (e.g., Program Files, /usr/local/bin) or registering system services."/>
+                </StackPanel>
+            </Border>
+
+        </StackPanel>
+    </ScrollViewer>
+
+</UserControl>
diff --git a/Engine/Source/Programs/UnrealToolboxExampleApplication/Pages/InstallUninstallPage.axaml.cs b/Engine/Source/Programs/UnrealToolboxExampleApplication/Pages/InstallUninstallPage.axaml.cs
new file mode 100644
--- /dev/null
+++ b/Engine/Source/Programs/UnrealToolboxExampleApplication/Pages/InstallUninstallPage.axaml.cs
@@ -0,0 +1,13 @@
+// Copyright Epic Games, Inc. All Rights Reserved.
+
+using Avalonia.Controls;
+
+namespace UnrealToolboxExampleApplication.Pages;
+
+public partial class InstallUninstallPage : UserControl
+{
+	public InstallUninstallPage()
+	{
+		InitializeComponent();
+	}
+}
diff --git a/Engine/Source/Programs/UnrealToolboxExampleApplication/Pages/ProjectSetupPage.axaml b/Engine/Source/Programs/UnrealToolboxExampleApplication/Pages/ProjectSetupPage.axaml
new file mode 100644
--- /dev/null
+++ b/Engine/Source/Programs/UnrealToolboxExampleApplication/Pages/ProjectSetupPage.axaml
@@ -0,0 +1,62 @@
+<UserControl xmlns="[REDACTED_PROFILE_URL]"
+             xmlns:x="http://schemas.microsoft.com/winfx/2006/xaml"
+             xmlns:local="using:UnrealToolboxExampleApplication"
+             x:Class="UnrealToolboxExampleApplication.Pages.ProjectSetupPage">
+
+    <ScrollViewer Padding="32 24">
+        <StackPanel Spacing="4" MaxWidth="700">
+
+            <TextBlock Classes="PageTitle" Text="Project Setup"/>
+
+            <TextBlock Classes="DocBody"
+                       Text="A Toolbox tool can be any executable — .NET, native C++, Electron, etc. This guide uses a .NET Avalonia app because that is what this example application is built with, but the packaging concepts apply to any technology."/>
+
+            <!-- .csproj structure -->
+            <TextBlock Classes="SectionHeader" Text="Project File Structure"/>
+
+            <TextBlock Classes="DocBody"
+                       Text="Below is the .csproj from this example application. The key points are WinExe output type (no console window), the Avalonia and FluentAvaloniaUI package references, and compiled bindings enabled."/>
+
+            <Border Classes="CodeBlock">
+                <SelectableTextBlock Text="{x:Static local:CodeSamples.CsprojStructure}"/>
+            </Border>
+
+            <!-- Single instance lock -->
+            <TextBlock Classes="SectionHeader" Text="Single-Instance Lock"/>
+
+            <TextBlock Classes="DocBody"
+                       Text="The Toolbox may launch your tool's runCommand at any time (e.g., when the user clicks it in the tray menu). If your tool is already running, a second instance should detect this and exit gracefully. A cross-platform file-lock approach works everywhere:"/>
+
+            <Border Classes="CodeBlock">
+                <SelectableTextBlock Text="{x:Static local:CodeSamples.SingleInstanceLock}"/>
+            </Border>
+
+            <TextBlock Classes="DocBody" Text="Usage in Main:"/>
+
+            <Border Classes="CodeBlock">
+                <SelectableTextBlock Text="{x:Static local:CodeSamples.SingleInstanceUsage}"/>
+            </Border>
+
+            <!-- Platform builds -->
+            <TextBlock Classes="SectionHeader" Text="Platform-Specific Builds"/>
+
+            <TextBlock Classes="DocBody"
+                       Text="Horde tool deployments target specific Runtime Identifiers (RIDs) like win-x64, osx-arm64, or linux-x64. Each platform needs its own Toolbox.json because executable names differ (e.g., .exe on Windows). Use conditional ItemGroups to select the right config at publish time:"/>
+
+            <Border Classes="CodeBlock">
+                <SelectableTextBlock Text="{x:Static local:CodeSamples.PlatformRidConditions}"/>
+            </Border>
+
+            <!-- Tip card -->
+            <Border Classes="DocCard" Margin="0 12 0 0">
+                <StackPanel Spacing="4">
+                    <TextBlock FontWeight="SemiBold" FontSize="14" Text="Tip: {ToolDir} Expansion"/>
+                    <TextBlock Classes="DocBody"
+                               Text="In Toolbox.json, the fileName field supports {ToolDir} as a placeholder that expands to the directory where the tool is installed. This is useful if your tool's executable is in a subdirectory: &quot;fileName&quot;: &quot;{ToolDir}/bin/mytool.exe&quot;"/>
+                </StackPanel>
+            </Border>
+
+        </StackPanel>
+    </ScrollViewer>
+
+</UserControl>
diff --git a/Engine/Source/Programs/UnrealToolboxExampleApplication/Pages/ProjectSetupPage.axaml.cs b/Engine/Source/Programs/UnrealToolboxExampleApplication/Pages/ProjectSetupPage.axaml.cs
new file mode 100644
--- /dev/null
+++ b/Engine/Source/Programs/UnrealToolboxExampleApplication/Pages/ProjectSetupPage.axaml.cs
@@ -0,0 +1,13 @@
+// Copyright Epic Games, Inc. All Rights Reserved.
+
+using Avalonia.Controls;
+
+namespace UnrealToolboxExampleApplication.Pages;
+
+public partial class ProjectSetupPage : UserControl
+{
+	public ProjectSetupPage()
+	{
+		InitializeComponent();
+	}
+}
diff --git a/Engine/Source/Programs/UnrealToolboxExampleApplication/Pages/PublishingPage.axaml b/Engine/Source/Programs/UnrealToolboxExampleApplication/Pages/PublishingPage.axaml
new file mode 100644
--- /dev/null
+++ b/Engine/Source/Programs/UnrealToolboxExampleApplication/Pages/PublishingPage.axaml
@@ -0,0 +1,106 @@
+<UserControl xmlns="[REDACTED_PROFILE_URL]"
+             xmlns:x="http://schemas.microsoft.com/winfx/2006/xaml"
+             xmlns:local="using:UnrealToolboxExampleApplication"
+             x:Class="UnrealToolboxExampleApplication.Pages.PublishingPage">
+
+    <ScrollViewer Padding="32 24">
+        <StackPanel Spacing="4" MaxWidth="700">
+
+            <TextBlock Classes="PageTitle" Text="Publishing to Horde"/>
+
+            <TextBlock Classes="DocBody"
+                       Text="Once your tool is built and packaged with a Toolbox.json, you need to publish it to the Horde server so the Toolbox can discover and distribute it."/>
+
+            <!-- Server-side configuration -->
+            <TextBlock Classes="SectionHeader" Text="Server-Side Configuration"/>
+
+            <Border Classes="DocCard">
+                <StackPanel Spacing="4">
+                    <TextBlock FontWeight="SemiBold" Text="ShowInToolbox"/>
+                    <TextBlock Classes="DocBody"
+                               Text="The Horde tool deployment must have the ShowInToolbox flag set to true. Without this, the Toolbox will not include the tool in its catalog even if it exists on the server."/>
+                </StackPanel>
+            </Border>
+
+            <Border Classes="DocCard">
+                <StackPanel Spacing="4">
+                    <TextBlock FontWeight="SemiBold" Text="Platforms"/>
+                    <TextBlock Classes="DocBody"
+                               Text="The deployment specifies which platforms it supports using a Platforms array of .NET Runtime Identifiers. The Toolbox matches these against the local machine to determine which deployments to offer."/>
+                </StackPanel>
+            </Border>
+
+            <!-- RID table -->
+            <TextBlock Classes="SectionHeader" Text="Platform RID Reference"/>
+
+            <Border Classes="DocCard">
+                <StackPanel Spacing="6">
+                    <TextBlock FontWeight="SemiBold" Text="Common Runtime Identifiers"/>
+                    <TextBlock Classes="DocBody" Text="any — Platform-independent (e.g., scripts, JARs)"/>
+                    <TextBlock Classes="DocBody" Text="win — Any Windows architecture"/>
+                    <TextBlock Classes="DocBody" Text="win-x64 — Windows 64-bit Intel/AMD"/>
+                    <TextBlock Classes="DocBody" Text="win-arm64 — Windows ARM64"/>
+                    <TextBlock Classes="DocBody" Text="osx-x64 — macOS Intel"/>
+                    <TextBlock Classes="DocBody" Text="osx-arm64 — macOS Apple Silicon"/>
+                    <TextBlock Classes="DocBody" Text="linux-x64 — Linux 64-bit Intel/AMD"/>
+                    <TextBlock Classes="DocBody" Text="linux-arm64 — Linux ARM64"/>
+                </StackPanel>
+            </Border>
+
+            <!-- Auto-update -->
+            <TextBlock Classes="SectionHeader" Text="Auto-Update Behavior"/>
+
+            <TextBlock Classes="DocBody"
+                       Text="The Toolbox polls the Horde server every 5 minutes for tool updates. When a new version is detected:"/>
+
+            <Border Classes="DocCard">
+                <StackPanel Spacing="4">
+                    <TextBlock Classes="DocBody" Text="1. The new version is downloaded and extracted to a staging directory"/>
+                    <TextBlock Classes="DocBody" Text="2. The old version's uninstallCommand is run (triggering the update handshake if the tool is running)"/>
+                    <TextBlock Classes="DocBody" Text="3. The old version directory is deleted"/>
+                    <TextBlock Classes="DocBody" Text="4. The new version's installCommand is run"/>
+                    <TextBlock Classes="DocBody" Text="5. Tools.json is updated with the new version"/>
+                </StackPanel>
+            </Border>
+
+            <Border Classes="DocCard">
+                <StackPanel Spacing="4">
+                    <TextBlock FontWeight="SemiBold" Text="Exceptions"/>
+                    <TextBlock Classes="DocBody"
+                               Text="Tools with manualInstall: true are not auto-updated — the user must trigger updates from the UI. Tools identified by msi-product-id are also excluded from auto-update since MSI installations are managed by Windows Installer."/>
+                </StackPanel>
+            </Border>
+
+            <!-- Local development -->
+            <TextBlock Classes="SectionHeader" Text="Local Development"/>
+
+            <TextBlock Classes="DocBody"
+                       Text="During development, you can test your tool locally without publishing to Horde. Run UnrealToolbox with the -LocalTool flag pointing to your tool's publish output directory:"/>
+
+            <Border Classes="CodeBlock">
+                <SelectableTextBlock Text="{x:Static local:CodeSamples.LocalToolFlag}"/>
+            </Border>
+
+            <TextBlock Classes="DocCaption"
+                       Text="The local tool appears in the catalog alongside server-published tools. It is not auto-updated."/>
+
+            <!-- Deployment checklist -->
+            <TextBlock Classes="SectionHeader" Text="Deployment Checklist"/>
+
+            <Border Classes="DocCard">
+                <StackPanel Spacing="4">
+                    <TextBlock FontWeight="SemiBold" Text="Before publishing, verify:"/>
+                    <TextBlock Classes="DocBody" Text="1. Toolbox.json is present at the root of your publish output"/>
+                    <TextBlock Classes="DocBody" Text="2. installCommand and uninstallCommand work correctly"/>
+                    <TextBlock Classes="DocBody" Text="3. The tool handles single-instance correctly (second launch exits)"/>
+                    <TextBlock Classes="DocBody" Text="4. The update handshake works (uninstall while running shows dialog)"/>
+                    <TextBlock Classes="DocBody" Text="5. Platform RIDs in the Horde deployment match your build targets"/>
+                    <TextBlock Classes="DocBody" Text="6. ShowInToolbox is enabled on the Horde tool deployment"/>
+                    <TextBlock Classes="DocBody" Text="7. Test with -LocalTool before publishing to the server"/>
+                </StackPanel>
+            </Border>
+
+        </StackPanel>
+    </ScrollViewer>
+
+</UserControl>
diff --git a/Engine/Source/Programs/UnrealToolboxExampleApplication/Pages/PublishingPage.axaml.cs b/Engine/Source/Programs/UnrealToolboxExampleApplication/Pages/PublishingPage.axaml.cs
new file mode 100644
--- /dev/null
+++ b/Engine/Source/Programs/UnrealToolboxExampleApplication/Pages/PublishingPage.axaml.cs
@@ -0,0 +1,13 @@
+// Copyright Epic Games, Inc. All Rights Reserved.
+
+using Avalonia.Controls;
+
+namespace UnrealToolboxExampleApplication.Pages;
+
+public partial class PublishingPage : UserControl
+{
+	public PublishingPage()
+	{
+		InitializeComponent();
+	}
+}
diff --git a/Engine/Source/Programs/UnrealToolboxExampleApplication/Pages/ThisExamplePage.axaml b/Engine/Source/Programs/UnrealToolboxExampleApplication/Pages/ThisExamplePage.axaml
new file mode 100644
--- /dev/null
+++ b/Engine/Source/Programs/UnrealToolboxExampleApplication/Pages/ThisExamplePage.axaml
@@ -0,0 +1,96 @@
+<UserControl xmlns="[REDACTED_PROFILE_URL]"
+             xmlns:x="http://schemas.microsoft.com/winfx/2006/xaml"
+             xmlns:local="using:UnrealToolboxExampleApplication"
+             x:Class="UnrealToolboxExampleApplication.Pages.ThisExamplePage">
+
+    <ScrollViewer Padding="32 24">
+        <StackPanel Spacing="4" MaxWidth="700">
+
+            <TextBlock Classes="PageTitle" Text="This Example Application"/>
+
+            <TextBlock Classes="DocBody"
+                       Text="This application is a real, deployed Toolbox tool. Below is an inventory of its source files and an explanation of how each piece fits together."/>
+
+            <!-- Source file inventory -->
+            <TextBlock Classes="SectionHeader" Text="Source File Inventory"/>
+
+            <Border Classes="DocCard">
+                <StackPanel Spacing="6">
+                    <TextBlock FontWeight="SemiBold" Text="Application Core"/>
+                    <TextBlock Classes="DocBody" Text="Program.cs — Entry point, single-instance lock, install/uninstall commands, platform-specific shortcuts"/>
+                    <TextBlock Classes="DocBody" Text="App.axaml / App.axaml.cs — Avalonia application setup, FluentAvalonia theme, update listener threads"/>
+                    <TextBlock Classes="DocBody" Text="MainWindow.axaml / MainWindow.axaml.cs — NavigationView shell, page routing and caching"/>
+                    <TextBlock Classes="DocBody" Text="UpdateDialog.axaml / UpdateDialog.axaml.cs — Modal dialog shown during the update handshake"/>
+                </StackPanel>
+            </Border>
+
+            <Border Classes="DocCard">
+                <StackPanel Spacing="6">
+                    <TextBlock FontWeight="SemiBold" Text="Guide Pages"/>
+                    <TextBlock Classes="DocBody" Text="Pages/WelcomePage.axaml — Overview and three-step summary"/>
+                    <TextBlock Classes="DocBody" Text="Pages/ProjectSetupPage.axaml — .csproj structure, single-instance lock, platform builds"/>
+                    <TextBlock Classes="DocBody" Text="Pages/ToolboxJsonPage.axaml — Complete Toolbox.json reference"/>
+                    <TextBlock Classes="DocBody" Text="Pages/InstallUninstallPage.axaml — Install examples and update handshake"/>
+                    <TextBlock Classes="DocBody" Text="Pages/PublishingPage.axaml — Horde deployment, RIDs, auto-update"/>
+                    <TextBlock Classes="DocBody" Text="Pages/ThisExamplePage.axaml — This page"/>
+                </StackPanel>
+            </Border>
+
+            <Border Classes="DocCard">
+                <StackPanel Spacing="6">
+                    <TextBlock FontWeight="SemiBold" Text="Shared Resources"/>
+                    <TextBlock Classes="DocBody" Text="Styles.axaml — DocCard, CodeBlock, PageTitle, SectionHeader, DocBody, DocCaption styles"/>
+                    <TextBlock Classes="DocBody" Text="CodeSamples.cs — Static code snippet strings referenced by pages via x:Static"/>
+                </StackPanel>
+            </Border>
+
+            <!-- Platform config comparison -->
+            <TextBlock Classes="SectionHeader" Text="Platform Configuration Comparison"/>
+
+            <TextBlock Classes="DocBody"
+                       Text="This app ships three Toolbox.json variants. The only difference is the executable name — Windows uses .exe, Unix platforms do not:"/>
+
+            <TextBlock Classes="DocBody" FontWeight="SemiBold" Margin="0 8 0 2" Text="Toolbox.win.json"/>
+            <Border Classes="CodeBlock">
+                <SelectableTextBlock Text="{x:Static local:CodeSamples.ToolboxWinJson}"/>
+            </Border>
+
+            <TextBlock Classes="DocBody" FontWeight="SemiBold" Margin="0 8 0 2" Text="Toolbox.osx.json / Toolbox.linux.json"/>
+            <Border Classes="CodeBlock">
+                <SelectableTextBlock Text="{x:Static local:CodeSamples.ToolboxUnixJson}"/>
+            </Border>
+
+            <!-- Update dialog -->
+            <TextBlock Classes="SectionHeader" Text="Update Dialog"/>
+
+            <TextBlock Classes="DocBody"
+                       Text="When the Toolbox needs to update this tool while it is running, it triggers the update handshake (named event on Windows, marker file on macOS/Linux). The App.axaml.cs background listener thread detects the signal and posts to the UI thread to show UpdateDialog — a simple Yes/No modal."/>
+
+            <Border Classes="DocCard">
+                <StackPanel Spacing="4">
+                    <TextBlock Classes="DocBody" Text="If the user clicks Yes, the application shuts down so the Toolbox can replace the files and run the new version's installCommand."/>
+                    <TextBlock Classes="DocBody" Text="If the user clicks No, the app signals back (via declined event or marker file) and the update is deferred until the user closes the app or accepts a future prompt."/>
+                </StackPanel>
+            </Border>
+
+            <!-- Interactive suggestions -->
+            <TextBlock Classes="SectionHeader" Text="Try It"/>
+
+            <Border Classes="DocCard">
+                <StackPanel Spacing="6">
+                    <TextBlock FontWeight="SemiBold" Text="Interactive Suggestions"/>
+                    <TextBlock Classes="DocBody"
+                               Text="Open the Unreal Toolbox settings window — this tool should appear in the tool catalog list."/>
+                    <TextBlock Classes="DocBody"
+                               Text="Look at the system tray — if this tool defines popupMenu items, they appear in the Toolbox tray context menu."/>
+                    <TextBlock Classes="DocBody"
+                               Text="Try running a second instance of this application — it should detect the lock file and exit immediately with an error message."/>
+                    <TextBlock Classes="DocBody"
+                               Text="Examine the source files side-by-side with this guide to see how each concept is implemented in practice."/>
+                </StackPanel>
+            </Border>
+
+        </StackPanel>
+    </ScrollViewer>
+
+</UserControl>
diff --git a/Engine/Source/Programs/UnrealToolboxExampleApplication/Pages/ThisExamplePage.axaml.cs b/Engine/Source/Programs/UnrealToolboxExampleApplication/Pages/ThisExamplePage.axaml.cs
new file mode 100644
--- /dev/null
+++ b/Engine/Source/Programs/UnrealToolboxExampleApplication/Pages/ThisExamplePage.axaml.cs
@@ -0,0 +1,13 @@
+// Copyright Epic Games, Inc. All Rights Reserved.
+
+using Avalonia.Controls;
+
+namespace UnrealToolboxExampleApplication.Pages;
+
+public partial class ThisExamplePage : UserControl
+{
+	public ThisExamplePage()
+	{
+		InitializeComponent();
+	}
+}
diff --git a/Engine/Source/Programs/UnrealToolboxExampleApplication/Pages/ToolboxJsonPage.axaml b/Engine/Source/Programs/UnrealToolboxExampleApplication/Pages/ToolboxJsonPage.axaml
new file mode 100644
--- /dev/null
+++ b/Engine/Source/Programs/UnrealToolboxExampleApplication/Pages/ToolboxJsonPage.axaml
@@ -0,0 +1,121 @@
+<UserControl xmlns="[REDACTED_PROFILE_URL]"
+             xmlns:x="http://schemas.microsoft.com/winfx/2006/xaml"
+             xmlns:local="using:UnrealToolboxExampleApplication"
+             x:Class="UnrealToolboxExampleApplication.Pages.ToolboxJsonPage">
+
+    <ScrollViewer Padding="32 24">
+        <StackPanel Spacing="4" MaxWidth="700">
+
+            <TextBlock Classes="PageTitle" Text="Toolbox.json Reference"/>
+
+            <TextBlock Classes="DocBody"
+                       Text="Every Toolbox tool must include a Toolbox.json file at the root of its deployment zip. This file tells the Toolbox how to install, uninstall, and run the tool, and optionally configures tray menu items and other behavior."/>
+
+            <!-- Commands -->
+            <TextBlock Classes="SectionHeader" Text="Commands"/>
+
+            <TextBlock Classes="DocBody"
+                       Text="Commands use the ToolCommand schema: an object with a fileName string and an optional arguments array of strings."/>
+
+            <Border Classes="DocCard">
+                <StackPanel Spacing="4">
+                    <TextBlock FontWeight="SemiBold" Text="installCommand"/>
+                    <TextBlock Classes="DocBody"
+                               Text="Runs after the tool is downloaded and extracted. Use this to create shortcuts, symlinks, register services, or perform any first-time setup. Receives the extracted directory as the working directory."/>
+                </StackPanel>
+            </Border>
+
+            <Border Classes="DocCard">
+                <StackPanel Spacing="4">
+                    <TextBlock FontWeight="SemiBold" Text="uninstallCommand"/>
+                    <TextBlock Classes="DocBody"
+                               Text="Runs before the tool directory is deleted. Use this to remove shortcuts, stop services, and clean up. The tool directory is still present when this runs."/>
+                </StackPanel>
+            </Border>
+
+            <Border Classes="DocCard">
+                <StackPanel Spacing="4">
+                    <TextBlock FontWeight="SemiBold" Text="runCommand"/>
+                    <TextBlock Classes="DocBody"
+                               Text="Optional. Launched when the user clicks the tool in the tray menu or opens it from the Toolbox settings. If not set, the tool has no direct launch action."/>
+                </StackPanel>
+            </Border>
+
+            <!-- Options -->
+            <TextBlock Classes="SectionHeader" Text="Options"/>
+
+            <Border Classes="DocCard">
+                <StackPanel Spacing="4">
+                    <TextBlock FontWeight="SemiBold" Text="manualInstall"/>
+                    <TextBlock Classes="DocBody"
+                               Text="Default: false. When true, the Toolbox downloads the tool but does not run the installCommand automatically. The user must trigger installation manually from the settings UI. Auto-update is also disabled."/>
+                </StackPanel>
+            </Border>
+
+            <Border Classes="DocCard">
+                <StackPanel Spacing="4">
+                    <TextBlock FontWeight="SemiBold" Text="requiresElevation"/>
+                    <TextBlock Classes="DocBody"
+                               Text="Default: false. When true, the installCommand and uninstallCommand are run with administrator/root privileges. On Windows, this triggers the UAC prompt."/>
+                </StackPanel>
+            </Border>
+
+            <Border Classes="DocCard">
+                <StackPanel Spacing="4">
+                    <TextBlock FontWeight="SemiBold" Text="hidden"/>
+                    <TextBlock Classes="DocBody"
+                               Text="Default: false. When true, the tool does not appear in the Toolbox settings UI. It is still installed and updated, but invisible to the user. Useful for background services."/>
+                </StackPanel>
+            </Border>
+
+            <!-- Tray Menu -->
+            <TextBlock Classes="SectionHeader" Text="Tray Menu (popupMenu)"/>
+
+            <TextBlock Classes="DocBody"
+                       Text="The popupMenu object is a single ToolMenuItem that appears in the Toolbox system tray context menu. For multiple items, use children to create a submenu. The ToolMenuItem schema:"/>
+
+            <Border Classes="DocCard">
+                <StackPanel Spacing="4">
+                    <TextBlock FontWeight="SemiBold" Text="ToolMenuItem Schema"/>
+                    <TextBlock Classes="DocBody" Text="label — Display text for the menu item"/>
+                    <TextBlock Classes="DocBody" Text="command — A ToolCommand to run when clicked (fileName + arguments)"/>
+                    <TextBlock Classes="DocBody" Text="children — Nested array of ToolMenuItem for submenus"/>
+                    <TextBlock Classes="DocCaption" Margin="0 4 0 0"
+                               Text="Each item should have either a command or children, not both."/>
+                </StackPanel>
+            </Border>
+
+            <!-- Platform files -->
+            <TextBlock Classes="SectionHeader" Text="Platform-Specific Files"/>
+
+            <TextBlock Classes="DocBody"
+                       Text="You can provide separate config files for each platform. At build/publish time, the appropriate variant is copied as Toolbox.json:"/>
+
+            <Border Classes="DocCard">
+                <StackPanel Spacing="4">
+                    <TextBlock Classes="DocBody" Text="Toolbox.win.json — Windows (executables use .exe suffix)"/>
+                    <TextBlock Classes="DocBody" Text="Toolbox.osx.json — macOS (no extension)"/>
+                    <TextBlock Classes="DocBody" Text="Toolbox.linux.json — Linux (no extension)"/>
+                </StackPanel>
+            </Border>
+
+            <!-- Full example -->
+            <TextBlock Classes="SectionHeader" Text="Full Annotated Example"/>
+
+            <TextBlock Classes="DocBody"
+                       Text="A complete Toolbox.json showing all fields. In practice, most tools only need installCommand, uninstallCommand, and optionally runCommand."/>
+
+            <Border Classes="CodeBlock">
+                <SelectableTextBlock Text="{x:Static local:CodeSamples.ToolboxJsonFull}"/>
+            </Border>
+
+            <TextBlock Classes="DocBody" Margin="0 8 0 0" Text="Minimal example (install and uninstall only):"/>
+
+            <Border Classes="CodeBlock">
+                <SelectableTextBlock Text="{x:Static local:CodeSamples.ToolboxJsonMinimal}"/>
+            </Border>
+
+        </StackPanel>
+    </ScrollViewer>
+
+</UserControl>
diff --git a/Engine/Source/Programs/UnrealToolboxExampleApplication/Pages/ToolboxJsonPage.axaml.cs b/Engine/Source/Programs/UnrealToolboxExampleApplication/Pages/ToolboxJsonPage.axaml.cs
new file mode 100644
--- /dev/null
+++ b/Engine/Source/Programs/UnrealToolboxExampleApplication/Pages/ToolboxJsonPage.axaml.cs
@@ -0,0 +1,13 @@
+// Copyright Epic Games, Inc. All Rights Reserved.
+
+using Avalonia.Controls;
+
+namespace UnrealToolboxExampleApplication.Pages;
+
+public partial class ToolboxJsonPage : UserControl
+{
+	public ToolboxJsonPage()
+	{
+		InitializeComponent();
+	}
+}
diff --git a/Engine/Source/Programs/UnrealToolboxExampleApplication/Pages/WelcomePage.axaml b/Engine/Source/Programs/UnrealToolboxExampleApplication/Pages/WelcomePage.axaml
new file mode 100644
--- /dev/null
+++ b/Engine/Source/Programs/UnrealToolboxExampleApplication/Pages/WelcomePage.axaml
@@ -0,0 +1,57 @@
+<UserControl xmlns="[REDACTED_PROFILE_URL]"
+             xmlns:x="http://schemas.microsoft.com/winfx/2006/xaml"
+             x:Class="UnrealToolboxExampleApplication.Pages.WelcomePage">
+
+    <ScrollViewer Padding="32 24">
+        <StackPanel Spacing="4" MaxWidth="700">
+
+            <Border CornerRadius="8" ClipToBounds="True" Margin="0 0 0 16">
+                <Image Source="/Resources/utbea_banner.png" Stretch="Uniform" MaxHeight="220"/>
+            </Border>
+
+            <TextBlock Classes="PageTitle" Text="Welcome to the UnrealToolbox Developer Guide"/>
+
+            <TextBlock Classes="DocBody"
+                       Text="UnrealToolbox is a Windows/macOS/Linux system tray application that manages Horde build tools and artifacts. Tools published to the Horde server are automatically discovered, downloaded, installed, and kept up to date by the Toolbox."/>
+
+            <TextBlock Classes="DocBody" Margin="0 8 0 0"
+                       Text="This guide walks you through publishing your own tool so it appears in the Toolbox catalog. Each page covers one step of the process."/>
+
+            <TextBlock Classes="SectionHeader" Text="The Three Steps"/>
+
+            <Border Classes="DocCard">
+                <StackPanel Spacing="4">
+                    <TextBlock FontWeight="SemiBold" FontSize="14" Text="1. Package Your Tool"/>
+                    <TextBlock Classes="DocBody"
+                               Text="Create a .NET (or native) application, add a single-instance lock, and set up platform-specific builds. The Project Setup page covers this in detail."/>
+                </StackPanel>
+            </Border>
+
+            <Border Classes="DocCard">
+                <StackPanel Spacing="4">
+                    <TextBlock FontWeight="SemiBold" FontSize="14" Text="2. Configure Toolbox.json"/>
+                    <TextBlock Classes="DocBody"
+                               Text="Add a Toolbox.json file to the root of your tool's zip archive. This tells the Toolbox how to install, uninstall, and run your tool, and optionally adds tray menu items."/>
+                </StackPanel>
+            </Border>
+
+            <Border Classes="DocCard">
+                <StackPanel Spacing="4">
+                    <TextBlock FontWeight="SemiBold" FontSize="14" Text="3. Publish to Horde"/>
+                    <TextBlock Classes="DocBody"
+                               Text="Upload the tool to the Horde server with the ShowInToolbox flag enabled. The Toolbox will discover it within 5 minutes and offer it to users."/>
+                </StackPanel>
+            </Border>
+
+            <Border Classes="DocCard" Margin="0 16 0 0">
+                <StackPanel Spacing="4">
+                    <TextBlock FontWeight="SemiBold" FontSize="14" Text="This application is itself a deployed Toolbox tool"/>
+                    <TextBlock Classes="DocCaption"
+                               Text="Every code sample in this guide is drawn from the source code of this example application. You can inspect the source alongside these pages to see a real, working implementation of each concept."/>
+                </StackPanel>
+            </Border>
+
+        </StackPanel>
+    </ScrollViewer>
+
+</UserControl>
diff --git a/Engine/Source/Programs/UnrealToolboxExampleApplication/Pages/WelcomePage.axaml.cs b/Engine/Source/Programs/UnrealToolboxExampleApplication/Pages/WelcomePage.axaml.cs
new file mode 100644
--- /dev/null
+++ b/Engine/Source/Programs/UnrealToolboxExampleApplication/Pages/WelcomePage.axaml.cs
@@ -0,0 +1,13 @@
+// Copyright Epic Games, Inc. All Rights Reserved.
+
+using Avalonia.Controls;
+
+namespace UnrealToolboxExampleApplication.Pages;
+
+public partial class WelcomePage : UserControl
+{
+	public WelcomePage()
+	{
+		InitializeComponent();
+	}
+}
diff --git a/Engine/Source/Programs/UnrealToolboxExampleApplication/Program.cs b/Engine/Source/Programs/UnrealToolboxExampleApplication/Program.cs
new file mode 100644
--- /dev/null
+++ b/Engine/Source/Programs/UnrealToolboxExampleApplication/Program.cs
@@ -0,0 +1,415 @@
+// Copyright Epic Games, Inc. All Rights Reserved.
+
+using System.Diagnostics;
+using System.Diagnostics.CodeAnalysis;
+using System.Reflection;
+using System.Runtime.InteropServices;
+using System.Runtime.Versioning;
+using Avalonia;
+
+namespace UnrealToolboxExampleApplication;
+
+static class Program
+{
+	public const string UpdateRequestEventName = "UnrealToolboxExampleApp-UpdateRequest";
+	public const string UpdateDeclinedEventName = "UnrealToolboxExampleApp-UpdateDeclined";
+
+	const string ProcessName = "UnrealToolboxExampleApplication";
+	const string AppDisplayName = "Unreal Toolbox Example Application";
+
+	[STAThread]
+	public static int Main(string[] args)
+	{
+		if (args.Any(x => x.Equals("--install", StringComparison.OrdinalIgnoreCase)))
+		{
+			return Install();
+		}
+		if (args.Any(x => x.Equals("--uninstall", StringComparison.OrdinalIgnoreCase)))
+		{
+			return Uninstall();
+		}
+
+		using SingleInstanceLock instanceLock = new SingleInstanceLock("UnrealToolboxExampleApplication");
+		if (!instanceLock.TryAcquire())
+		{
+			Console.Error.WriteLine("Another instance is already running.");
+			return 1;
+		}
+
+		BuildAvaloniaApp().StartWithClassicDesktopLifetime(args);
+		return 0;
+	}
+
+	public static AppBuilder BuildAvaloniaApp()
+	{
+		return AppBuilder.Configure<App>()
+			.UsePlatformDetect()
+			.LogToTrace();
+	}
+
+	static int Install()
+	{
+		try
+		{
+			string exePath = Environment.ProcessPath ?? throw new InvalidOperationException("Cannot determine process path.");
+
+			if (OperatingSystem.IsWindows())
+			{
+				InstallWindows(exePath);
+			}
+			else if (OperatingSystem.IsMacOS())
+			{
+				InstallMacOS(exePath);
+			}
+			else if (OperatingSystem.IsLinux())
+			{
+				InstallLinux(exePath);
+			}
+
+			return 0;
+		}
+		catch (Exception ex)
+		{
+			Console.Error.WriteLine($"Install failed: {ex.Message}");
+			return 1;
+		}
+	}
+
+	static int Uninstall()
+	{
+		try
+		{
+			if (!RequestRunningInstanceShutdown())
+			{
+				return 1;
+			}
+
+			if (OperatingSystem.IsWindows())
+			{
+				UninstallWindows();
+			}
+			else if (OperatingSystem.IsMacOS())
+			{
+				UninstallMacOS();
+			}
+			else if (OperatingSystem.IsLinux())
+			{
+				UninstallLinux();
+			}
+
+			return 0;
+		}
+		catch (Exception ex)
+		{
+			Console.Error.WriteLine($"Uninstall failed: {ex.Message}");
+			return 1;
+		}
+	}
+
+	#region Running instance shutdown
+
+	static bool RequestRunningInstanceShutdown()
+	{
+		Process[] instances = Process.GetProcessesByName(ProcessName)
+			.Where(p => p.Id != Environment.ProcessId)
+			.ToArray();
+
+		if (instances.Length == 0)
+		{
+			return true;
+		}
+
+		if (OperatingSystem.IsWindows())
+		{
+			return RequestShutdownViaNamedEvent(instances);
+		}
+		else
+		{
+			return RequestShutdownViaMarkerFile(instances);
+		}
+	}
+
+	static bool RequestShutdownViaNamedEvent(Process[] instances)
+	{
+		using EventWaitHandle requestEvent = new EventWaitHandle(false, EventResetMode.AutoReset, UpdateRequestEventName);
+		using EventWaitHandle declinedEvent = new EventWaitHandle(false, EventResetMode.AutoReset, UpdateDeclinedEventName);
+
+		Console.WriteLine("Requesting running instance to close for update...");
+		requestEvent.Set();
+
+		using ManualResetEvent allExitedEvent = new ManualResetEvent(false);
+		Thread waitThread = new Thread(() =>
+		{
+			foreach (Process p in instances)
+			{
+				try
+				{
+					p.WaitForExit();
+				}
+				catch
+				{
+				}
+				finally
+				{
+					p.Dispose();
+				}
+			}
+			allExitedEvent.Set();
+		});
+		waitThread.IsBackground = true;
+		waitThread.Start();
+
+		int index = WaitHandle.WaitAny(new WaitHandle[] { allExitedEvent, declinedEvent }, TimeSpan.FromSeconds(60));
+		if (index != 0)
+		{
+			string reason = index == 1 ? "User declined the update." : "Timed out waiting for user response.";
+			Console.Error.WriteLine($"Uninstall aborted: {reason}");
+			return false;
+		}
+
+		Console.WriteLine("Running instance closed.");
+		return true;
+	}
+
+	static bool RequestShutdownViaMarkerFile(Process[] instances)
+	{
+		// On Mac/Linux, use a temp file to ask the running instance to show the update dialog.
+		// The running app polls for the request file. If the user declines, it writes a declined file.
+		string requestFilePath = GetRequestFilePath();
+		string declinedFilePath = GetDeclinedFilePath();
+
+		if (File.Exists(declinedFilePath))
+		{
+			File.Delete(declinedFilePath);
+		}
+
+		Console.WriteLine("Writing update request marker file...");
+		File.WriteAllText(requestFilePath, "update");
+
+		// Poll for process exit or declined file
+		DateTime deadline = DateTime.UtcNow.AddSeconds(60);
+		while (DateTime.UtcNow < deadline)
+		{
+			if (File.Exists(declinedFilePath))
+			{
+				try
+				{
+					File.Delete(declinedFilePath);
+				}
+				catch
+				{
+				}
+				Console.Error.WriteLine("Uninstall aborted: User declined the update.");
+				foreach (Process p in instances)
+				{
+					p.Dispose();
+				}
+				return false;
+			}
+
+			bool allExited = true;
+			foreach (Process p in instances)
+			{
+				try
+				{
+					if (!p.HasExited)
+					{
+						allExited = false;
+						break;
+					}
+				}
+				catch
+				{
+				}
+			}
+
+			if (allExited)
+			{
+				foreach (Process p in instances)
+				{
+					p.Dispose();
+				}
+				Console.WriteLine("Running instance closed.");
+				return true;
+			}
+
+			Thread.Sleep(500);
+		}
+
+		// Clean up on timeout
+		try
+		{
+			File.Delete(requestFilePath);
+		}
+		catch
+		{
+		}
+		foreach (Process p in instances)
+		{
+			p.Dispose();
+		}
+		Console.Error.WriteLine("Uninstall aborted: Timed out waiting for user response.");
+		return false;
+	}
+
+	public static string GetRequestFilePath()
+	{
+		return Path.Combine(Path.GetTempPath(), "UnrealToolboxExampleApp-UpdateRequest");
+	}
+
+	public static string GetDeclinedFilePath()
+	{
+		return Path.Combine(Path.GetTempPath(), "UnrealToolboxExampleApp-UpdateDeclined");
+	}
+
+	#endregion
+
+	#region Windows
+
+	[SupportedOSPlatform("windows")]
+	[UnconditionalSuppressMessage("Trimming", "IL2057", Justification = "COM type WScript.Shell is resolved at runtime by the OS, not subject to trimming.")]
+	[UnconditionalSuppressMessage("Trimming", "IL2072", Justification = "COM type WScript.Shell is resolved at runtime by the OS, not subject to trimming.")]
+	[UnconditionalSuppressMessage("Trimming", "IL2075", Justification = "COM interop members are provided by the OS, not trimmed from the assembly.")]
+	static void InstallWindows(string exePath)
+	{
+		string shortcutPath = GetWindowsShortcutPath();
+
+		Type shellType = Type.GetTypeFromProgID("WScript.Shell") ?? throw new InvalidOperationException("WScript.Shell is not available.");
+		object shell = Activator.CreateInstance(shellType)!;
+		object shortcut = shellType.InvokeMember("CreateShortcut", BindingFlags.InvokeMethod, null, shell, new object[] { shortcutPath })!;
+		Type shortcutType = shortcut.GetType();
+		shortcutType.InvokeMember("TargetPath", BindingFlags.SetProperty, null, shortcut, new object[] { exePath });
+		shortcutType.InvokeMember("WorkingDirectory", BindingFlags.SetProperty, null, shortcut, new object[] { Path.GetDirectoryName(exePath)! });
+		shortcutType.InvokeMember("Description", BindingFlags.SetProperty, null, shortcut, new object[] { AppDisplayName });
+		shortcutType.InvokeMember("Save", BindingFlags.InvokeMethod, null, shortcut, null);
+
+		Marshal.ReleaseComObject(shortcut);
+		Marshal.ReleaseComObject(shell);
+
+		Console.WriteLine($"Shortcut created: {shortcutPath}");
+	}
+
+	[SupportedOSPlatform("windows")]
+	static void UninstallWindows()
+	{
+		string shortcutPath = GetWindowsShortcutPath();
+		if (File.Exists(shortcutPath))
+		{
+			File.Delete(shortcutPath);
+			Console.WriteLine($"Shortcut removed: {shortcutPath}");
+		}
+	}
+
+	static string GetWindowsShortcutPath()
+	{
+		string programsFolder = Environment.GetFolderPath(Environment.SpecialFolder.Programs);
+		return Path.Combine(programsFolder, $"{AppDisplayName}.lnk");
+	}
+
+	#endregion
+
+	#region macOS
+
+	static void InstallMacOS(string exePath)
+	{
+		string applicationsDir = Path.Combine(
+			Environment.GetFolderPath(Environment.SpecialFolder.UserProfile),
+			"Applications");
+		Directory.CreateDirectory(applicationsDir);
+
+		string linkPath = Path.Combine(applicationsDir, ProcessName);
+		if (File.Exists(linkPath))
+		{
+			File.Delete(linkPath);
+		}
+
+		File.CreateSymbolicLink(linkPath, exePath);
+		Console.WriteLine($"Symlink created: {linkPath} -> {exePath}");
+	}
+
+	static void UninstallMacOS()
+	{
+		string linkPath = Path.Combine(
+			Environment.GetFolderPath(Environment.SpecialFolder.UserProfile),
+			"Applications",
+			ProcessName);
+
+		if (Path.Exists(linkPath))
+		{
+			File.Delete(linkPath);
+			Console.WriteLine($"Symlink removed: {linkPath}");
+		}
+	}
+
+	#endregion
+
+	#region Linux
+
+	static void InstallLinux(string exePath)
+	{
+		string applicationsDir = Path.Combine(
+			Environment.GetFolderPath(Environment.SpecialFolder.UserProfile),
+			".local", "share", "applications");
+		Directory.CreateDirectory(applicationsDir);
+
+		string desktopFilePath = Path.Combine(applicationsDir, $"{ProcessName}.desktop");
+		string content =
+			"[Desktop Entry]\n" +
+			$"Name={AppDisplayName}\n" +
+			$"Exec={exePath}\n" +
+			"Type=Application\n" +
+			"Terminal=false\n" +
+			$"Comment={AppDisplayName}\n";
+
+		File.WriteAllText(desktopFilePath, content);
+		Console.WriteLine($"Desktop entry created: {desktopFilePath}");
+	}
+
+	static void UninstallLinux()
+	{
+		string desktopFilePath = Path.Combine(
+			Environment.GetFolderPath(Environment.SpecialFolder.UserProfile),
+			".local", "share", "applications",
+			$"{ProcessName}.desktop");
+
+		if (File.Exists(desktopFilePath))
+		{
+			File.Delete(desktopFilePath);
+			Console.WriteLine($"Desktop entry removed: {desktopFilePath}");
+		}
+	}
+
+	#endregion
+}
+
+class SingleInstanceLock : IDisposable
+{
+	readonly string _lockFilePath;
+	FileStream? _lockFileStream;
+
+	public SingleInstanceLock(string appName)
+	{
+		string lockDir = Path.Combine(Environment.GetFolderPath(Environment.SpecialFolder.LocalApplicationData), "Epic Games", appName);
+		Directory.CreateDirectory(lockDir);
+		_lockFilePath = Path.Combine(lockDir, ".lock");
+	}
+
+	public bool TryAcquire()
+	{
+		try
+		{
+			_lockFileStream = new FileStream(_lockFilePath, FileMode.OpenOrCreate, FileAccess.ReadWrite, FileShare.None);
+			return true;
+		}
+		catch (IOException)
+		{
+			return false;
+		}
+	}
+
+	public void Dispose()
+	{
+		_lockFileStream?.Dispose();
+		_lockFileStream = null;
+	}
+}
diff --git a/Engine/Source/Programs/UnrealToolboxExampleApplication/Styles.axaml b/Engine/Source/Programs/UnrealToolboxExampleApplication/Styles.axaml
new file mode 100644
--- /dev/null
+++ b/Engine/Source/Programs/UnrealToolboxExampleApplication/Styles.axaml
@@ -0,0 +1,50 @@
+<Styles xmlns="[REDACTED_PROFILE_URL]"
+		xmlns:x="http://schemas.microsoft.com/winfx/2006/xaml">
+
+	<Style Selector="Border.DocCard">
+		<Setter Property="BorderBrush" Value="{DynamicResource CardStrokeColorDefaultBrush}"/>
+		<Setter Property="BorderThickness" Value="1"/>
+		<Setter Property="Background" Value="{DynamicResource CardBackgroundFillColorDefaultBrush}"/>
+		<Setter Property="TextElement.Foreground" Value="{DynamicResource TextFillColorPrimaryBrush}"/>
+		<Setter Property="CornerRadius" Value="{DynamicResource ControlCornerRadius}"/>
+		<Setter Property="BackgroundSizing" Value="InnerBorderEdge"/>
+		<Setter Property="Padding" Value="16 14"/>
+		<Setter Property="Margin" Value="0 4"/>
+	</Style>
+
+	<Style Selector="Border.CodeBlock">
+		<Setter Property="Background" Value="#1E1E1E"/>
+		<Setter Property="CornerRadius" Value="4"/>
+		<Setter Property="Padding" Value="16 12"/>
+		<Setter Property="Margin" Value="0 6"/>
+	</Style>
+	<Style Selector="Border.CodeBlock SelectableTextBlock">
+		<Setter Property="FontFamily" Value="Cascadia Mono, Consolas, Menlo, monospace"/>
+		<Setter Property="FontSize" Value="12"/>
+		<Setter Property="Foreground" Value="#D4D4D4"/>
+		<Setter Property="TextWrapping" Value="NoWrap"/>
+	</Style>
+
+	<Style Selector="TextBlock.PageTitle">
+		<Setter Property="Theme" Value="{DynamicResource SubtitleTextBlockStyle}"/>
+		<Setter Property="Margin" Value="0 0 0 12"/>
+	</Style>
+
+	<Style Selector="TextBlock.SectionHeader">
+		<Setter Property="Theme" Value="{DynamicResource BodyStrongTextBlockStyle}"/>
+		<Setter Property="Margin" Value="0 20 0 8"/>
+	</Style>
+
+	<Style Selector="TextBlock.DocBody">
+		<Setter Property="FontSize" Value="13"/>
+		<Setter Property="TextWrapping" Value="Wrap"/>
+		<Setter Property="Margin" Value="0 2"/>
+	</Style>
+
+	<Style Selector="TextBlock.DocCaption">
+		<Setter Property="Theme" Value="{DynamicResource CaptionTextBlockStyle}"/>
+		<Setter Property="FontWeight" Value="Light"/>
+		<Setter Property="Margin" Value="0 2"/>
+	</Style>
+
+</Styles>
diff --git a/Engine/Source/Programs/UnrealToolboxExampleApplication/Toolbox.linux.json b/Engine/Source/Programs/UnrealToolboxExampleApplication/Toolbox.linux.json
new file mode 100644
--- /dev/null
+++ b/Engine/Source/Programs/UnrealToolboxExampleApplication/Toolbox.linux.json
@@ -0,0 +1,19 @@
+{
+  "installCommand": {
+    "fileName": "UnrealToolboxExampleApplication",
+    "arguments": [ "--install" ]
+  },
+  "uninstallCommand": {
+    "fileName": "UnrealToolboxExampleApplication",
+    "arguments": [ "--uninstall" ]
+  },
+  "runCommand": {
+    "fileName": "UnrealToolboxExampleApplication"
+  },
+  "popupMenu": {
+    "label": "Open Developer Guide",
+    "command": {
+      "fileName": "UnrealToolboxExampleApplication"
+    }
+  }
+}
diff --git a/Engine/Source/Programs/UnrealToolboxExampleApplication/Toolbox.osx.json b/Engine/Source/Programs/UnrealToolboxExampleApplication/Toolbox.osx.json
new file mode 100644
--- /dev/null
+++ b/Engine/Source/Programs/UnrealToolboxExampleApplication/Toolbox.osx.json
@@ -0,0 +1,19 @@
+{
+  "installCommand": {
+    "fileName": "UnrealToolboxExampleApplication",
+    "arguments": [ "--install" ]
+  },
+  "uninstallCommand": {
+    "fileName": "UnrealToolboxExampleApplication",
+    "arguments": [ "--uninstall" ]
+  },
+  "runCommand": {
+    "fileName": "UnrealToolboxExampleApplication"
+  },
+  "popupMenu": {
+    "label": "Open Developer Guide",
+    "command": {
+      "fileName": "UnrealToolboxExampleApplication"
+    }
+  }
+}
diff --git a/Engine/Source/Programs/UnrealToolboxExampleApplication/Toolbox.win.json b/Engine/Source/Programs/UnrealToolboxExampleApplication/Toolbox.win.json
new file mode 100644
--- /dev/null
+++ b/Engine/Source/Programs/UnrealToolboxExampleApplication/Toolbox.win.json
@@ -0,0 +1,19 @@
+{
+  "installCommand": {
+    "fileName": "UnrealToolboxExampleApplication.exe",
+    "arguments": [ "--install" ]
+  },
+  "uninstallCommand": {
+    "fileName": "UnrealToolboxExampleApplication.exe",
+    "arguments": [ "--uninstall" ]
+  },
+  "runCommand": {
+    "fileName": "UnrealToolboxExampleApplication.exe"
+  },
+  "popupMenu": {
+    "label": "Open Developer Guide",
+    "command": {
+      "fileName": "UnrealToolboxExampleApplication.exe"
+    }
+  }
+}
diff --git a/Engine/Source/Programs/UnrealToolboxExampleApplication/UnrealToolboxExampleApplication.csproj b/Engine/Source/Programs/UnrealToolboxExampleApplication/UnrealToolboxExampleApplication.csproj
new file mode 100644
--- /dev/null
+++ b/Engine/Source/Programs/UnrealToolboxExampleApplication/UnrealToolboxExampleApplication.csproj
@@ -0,0 +1,57 @@
+<Project Sdk="Microsoft.NET.Sdk">
+
+  <PropertyGroup>
+    <OutputType>WinExe</OutputType>
+    <TargetFramework>net10.0</TargetFramework>
+    <Nullable>enable</Nullable>
+    <ImplicitUsings>enable</ImplicitUsings>
+    <AvaloniaUseCompiledBindingsByDefault>true</AvaloniaUseCompiledBindingsByDefault>
+    <ApplicationIcon>Resources/utbea.ico</ApplicationIcon>
+    <SelfContained>true</SelfContained>
+    <PublishSingleFile>true</PublishSingleFile>
+    <PublishTrimmed>true</PublishTrimmed>
+    <PublishReadyToRun>true</PublishReadyToRun>
+  </PropertyGroup>
+
+  <ItemGroup>
+    <AvaloniaResource Include="Resources\**" />
+  </ItemGroup>
+
+  <ItemGroup>
+    <PackageReference Include="Avalonia" Version="11.3.12" />
+    <PackageReference Include="Avalonia.Desktop" Version="11.3.12" />
+    <PackageReference Include="Avalonia.Themes.Fluent" Version="11.3.12" />
+    <PackageReference Include="Avalonia.Fonts.Inter" Version="11.3.12" />
+    <PackageReference Include="FluentAvaloniaUI" Version="2.5.0" />
+  </ItemGroup>
+
+  <!-- Platform-specific Toolbox.json: copy the matching variant as Toolbox.json at publish time -->
+  <ItemGroup>
+    <None Include="Toolbox.win.json" CopyToPublishDirectory="Never" />
+    <None Include="Toolbox.osx.json" CopyToPublishDirectory="Never" />
+    <None Include="Toolbox.linux.json" CopyToPublishDirectory="Never" />
+  </ItemGroup>
+
+  <!-- Default to Windows for local build output (dotnet build without -r) -->
+  <ItemGroup Condition="'$(RuntimeIdentifier)' == '' Or $(RuntimeIdentifier.StartsWith('win'))">
+    <Content Include="Toolbox.win.json" Link="Toolbox.json">
+      <CopyToOutputDirectory>PreserveNewest</CopyToOutputDirectory>
+      <CopyToPublishDirectory>PreserveNewest</CopyToPublishDirectory>
+    </Content>
+  </ItemGroup>
+
+  <ItemGroup Condition="$(RuntimeIdentifier.StartsWith('osx'))">
+    <Content Include="Toolbox.osx.json" Link="Toolbox.json">
+      <CopyToOutputDirectory>PreserveNewest</CopyToOutputDirectory>
+      <CopyToPublishDirectory>PreserveNewest</CopyToPublishDirectory>
+    </Content>
+  </ItemGroup>
+
+  <ItemGroup Condition="$(RuntimeIdentifier.StartsWith('linux'))">
+    <Content Include="Toolbox.linux.json" Link="Toolbox.json">
+      <CopyToOutputDirectory>PreserveNewest</CopyToOutputDirectory>
+      <CopyToPublishDirectory>PreserveNewest</CopyToPublishDirectory>
+    </Content>
+  </ItemGroup>
+
+</Project>
diff --git a/Engine/Source/Programs/UnrealToolboxExampleApplication/UpdateDialog.axaml b/Engine/Source/Programs/UnrealToolboxExampleApplication/UpdateDialog.axaml
new file mode 100644
--- /dev/null
+++ b/Engine/Source/Programs/UnrealToolboxExampleApplication/UpdateDialog.axaml
@@ -0,0 +1,19 @@
+<Window xmlns="[REDACTED_PROFILE_URL]"
+        xmlns:x="http://schemas.microsoft.com/winfx/2006/xaml"
+        x:Class="UnrealToolboxExampleApplication.UpdateDialog"
+        Title="Update Available"
+        Width="400"
+        Height="180"
+        WindowStartupLocation="CenterOwner"
+        CanResize="False"
+        ShowInTaskbar="False">
+    <StackPanel Margin="24" Spacing="16">
+        <TextBlock Text="An update is available. Close the application to continue?"
+                   TextWrapping="Wrap"
+                   FontSize="14" />
+        <StackPanel Orientation="Horizontal" HorizontalAlignment="Right" Spacing="8">
+            <Button Content="Yes" Width="80" Click="OnYesClicked" />
+            <Button Content="No" Width="80" Click="OnNoClicked" />
+        </StackPanel>
+    </StackPanel>
+</Window>
diff --git a/Engine/Source/Programs/UnrealToolboxExampleApplication/UpdateDialog.axaml.cs b/Engine/Source/Programs/UnrealToolboxExampleApplication/UpdateDialog.axaml.cs
new file mode 100644
--- /dev/null
+++ b/Engine/Source/Programs/UnrealToolboxExampleApplication/UpdateDialog.axaml.cs
@@ -0,0 +1,24 @@
+// Copyright Epic Games, Inc. All Rights Reserved.
+
+using Avalonia.Controls;
+using Avalonia.Interactivity;
+
+namespace UnrealToolboxExampleApplication;
+
+public partial class UpdateDialog : Window
+{
+	public UpdateDialog()
+	{
+		InitializeComponent();
+	}
+
+	void OnYesClicked(object? sender, RoutedEventArgs e)
+	{
+		Close(true);
+	}
+
+	void OnNoClicked(object? sender, RoutedEventArgs e)
+	{
+		Close(false);
+	}
+}
